Add a downmixer and sample rate converter to AudioLink Submix buffer listener.
- Adds a BufferedSubmixListener subclass that does sample rate conversion and channel downmixing on top of a submix buffer listener client.
- Adds automated test to verify functionality in multithreaded environment

[RN] Added an AudioLinkEngine/DownmixedBufferedSubmixListener class to get an audio feed of a selected submix in the audio signal graph that will convert from the current mixer format to a specified number of channels and sample rate.
#Synced-CL 48392523

diff --git a/Engine/Plugins/Tests/AudioCQTests/Source/AudioCQTests/AudioCQTests.Build.cs b/Engine/Plugins/Tests/AudioCQTests/Source/AudioCQTests/AudioCQTests.Build.cs
--- a/Engine/Plugins/Tests/AudioCQTests/Source/AudioCQTests/AudioCQTests.Build.cs
+++ b/Engine/Plugins/Tests/AudioCQTests/Source/AudioCQTests/AudioCQTests.Build.cs
@@ -1,6 +1,7 @@
 // Copyright Epic Games, Inc. All Rights Reserved.
 
 using UnrealBuildTool;
+using UnrealBuildTool.Rules;
 
 public class AudioCQTests : ModuleRules
 {
@@ -11,10 +12,13 @@ public class AudioCQTests : ModuleRules
 
 		PrivateDependencyModuleNames.AddRange(
 			new string[] {
+					"AudioLinkEngine",
 					"AudioMixer",
 					"Core",
+					"CoreUObject",
 					"CQTest",
 					"Engine",
+					"SignalProcessing",
 				 }
 			);
 	}
diff --git a/Engine/Plugins/Tests/AudioCQTests/Source/AudioCQTests/Private/BufferedSubmixListenerTests.cpp b/Engine/Plugins/Tests/AudioCQTests/Source/AudioCQTests/Private/BufferedSubmixListenerTests.cpp
new file mode 100644
--- /dev/null
+++ b/Engine/Plugins/Tests/AudioCQTests/Source/AudioCQTests/Private/BufferedSubmixListenerTests.cpp
@@ -0,0 +1,334 @@
+// Copyright Epic Games, Inc. All Rights Reserved.
+
+#if WITH_DEV_AUTOMATION_TESTS
+
+#include "AudioDevice.h"
+#include "Misc/AutomationTest.h"
+#include "Components/AudioComponent.h"
+#include "CoreMinimal.h"
+#include "CQTest.h"
+#include "DownmixedBufferedSubmixListener.h"
+#include "Engine/World.h"
+#include "HAL/CriticalSection.h"
+#include "HAL/PlatformProcess.h"
+#include "HAL/ThreadSafeBool.h"
+#include "Math/UnrealMathUtility.h"
+#include "Sound/SoundBase.h"
+#include "Sound/SoundSubmix.h"
+#include "Tasks/Pipe.h"
+#include "TestSineSynthComponent.h"
+
+// This class creates a task that polls every 10ms and asks the BufferedSubmixListener for more data
+// then analyzes it to make sure it is the intended sine wave
+class FCQTestPoller
+{
+public:
+	FCQTestPoller(FAudioDevice* InAudioDevice, USoundSubmix* InSubmix, int32 InFrameCount, int32 InChannelCount, int32 InSampleRate)
+		: Pipe(TEXT("CQTestPoller")) 
+		, AudioDevice(InAudioDevice)
+		, Submix(InSubmix)
+		, FrameCount(InFrameCount)
+		, ChannelCount(InChannelCount)
+		, SampleRate(InSampleRate)
+	{
+		const FString VoiceChatSimName(TEXT("VoiceChatSimTest"));
+		Listener = MakeShared<FDownmixedBufferedSubmixListener>(InChannelCount, InSampleRate, &VoiceChatSimName);
+		Listener->OnSubmixBufferWritten.BindRaw(this, &FCQTestPoller::HandleOnSubmixBufferWritten);
+		TestBuffer.Reserve(InFrameCount * InChannelCount);
+	}
+
+	// Start polling loop in a task
+	void Start()
+	{
+		bStop = false;
+
+		if (AudioDevice && Listener.IsValid())
+		{
+			Listener->Start(AudioDevice, Submix);
+		}
+
+		// Launch a "thread-like" sequential task in the Pipe.
+		PollTask = Pipe.Launch(TEXT("CQTestPoller::Run"), [this]()
+			{
+				while (!bStop)
+				{
+					// Sleep ~10ms to simulate input audio driver task scheduling
+					FPlatformProcess::Sleep(0.01f); // 10ms
+					PollListenerQueue();
+				}
+			});
+	}
+
+	// Signal stop and wait for the task and pipe to quiesce
+	void Stop()
+	{
+		bStop = true;
+
+		if (Listener.IsValid())
+		{
+			if (AudioDevice)
+			{
+				Listener->Stop(AudioDevice);
+			}
+
+			Listener->OnSubmixBufferWritten.Unbind();
+			Listener.Reset();
+		}
+
+		AudioDevice = nullptr;
+
+		PollTask.Wait();
+		Pipe.WaitUntilEmpty(); // ensure pipe has no more work
+	}
+
+	// This function is called back by the BufferedSubmixListener when it has received a buffer
+	// Here we are just checking that there is actual signal coming in so we can start the test
+	void HandleOnSubmixBufferWritten(const float* SampleBuffer, int32 InFrameCount, int32 InNumChannels)
+	{
+		BuffersReceived++;
+
+		TArrayView<const float>ViewBuffer{ SampleBuffer, InFrameCount * InNumChannels };
+
+		// Compute a simple RMS without regard for channels, just to make sure there is a signal present
+		double SumSquares = 0.0;
+		for (int32 Sample = 0; Sample < InFrameCount * InNumChannels; Sample++)
+		{
+			const double SampleValue = static_cast<double>(SampleBuffer[Sample]);
+			SumSquares += SampleValue * SampleValue;
+		}
+		const double MeanSquare = (InFrameCount * InNumChannels > 0) ? (SumSquares / static_cast<double>(InFrameCount * InNumChannels)) : 0.0;
+		const float Rms = static_cast<float>(FMath::Sqrt(MeanSquare));
+
+		LatestRMS.store(Rms);
+	}
+
+	// Functions to access atomic values during the test process
+	void SetTestSineFrequency(float ExpectedFreq) { ExpectedSineFreq = ExpectedFreq; }
+	int32 GetBuffersReceived() { return BuffersReceived.load(); }
+	float GetLatestRMS() { return LatestRMS.load(); }
+	int32 GetSuccessfulTestCount() { return SuccessfulTests.load(); }
+	int32 GetFailedTestCount() { return FailedTests.load(); }
+	void ClearSetupValues() { BuffersReceived.store(0); LatestRMS.store(0.0f); }
+	void ClearTestValues() { SuccessfulTests.store(0); FailedTests.store(0); }
+
+private:
+	// This function is called on a 10ms callback to simulate VoiceChat timing
+	// This fetches a 10ms buffer of data from the BufferedSubmixListener and checks if it is actually a sine wave
+	void PollListenerQueue()
+	{
+		if (Listener.IsValid() && Listener->GetNumAvailableSamples() >= FrameCount * ChannelCount)
+		{
+			TestBuffer.SetNumUninitialized(FrameCount * ChannelCount);
+			float* FilledBuffer = TestBuffer.GetData();
+			int32 SamplesPopped = 0;
+			Listener->GetBuffer(TestBuffer.GetData(), FrameCount * ChannelCount, SamplesPopped);
+
+			bool TestPasses = ChannelCount > 0;
+
+			for (int32 Channel = 0; Channel < ChannelCount; Channel++)
+			{
+				float Amplitude = 0.0f;
+				float Phase = 0.0f;
+				TestPasses &= IsSineAtFrequency(FilledBuffer, SamplesPopped / ChannelCount, SampleRate, ChannelCount, Channel, ExpectedSineFreq, Amplitude, Phase);
+			}
+
+			if (TestPasses)
+			{
+				SuccessfulTests++;
+			}
+			else
+			{
+				FailedTests++;
+			}
+		}
+	}
+
+	// Returns true if buffer is well-fit by a sine at ExpectedFreq within tolerances.
+	// Also outputs estimated amplitude and phase.
+	bool IsSineAtFrequency(const float* Buffer, int32 InNumFrames, float InSampleRate, int32 InNumChannels, int32 ChannelIdx,
+		float ExpectedFreq, float& OutAmplitude, float& OutPhase, float RMSErrorTolerance = 1e-2f)
+	{
+		if (!Buffer || InNumFrames <= 0 || InSampleRate <= 0.f || ExpectedFreq <= 0.f || InNumChannels <= 0 || ChannelIdx >= InNumChannels)
+		{
+			return false;
+		}
+
+		const float TwoPi = UE_TWO_PI;
+		const float Omega = TwoPi * ExpectedFreq / InSampleRate;
+
+		// Correlate against cos(omega n) and sin(omega n)
+		double SumCos = 0.0;
+		double SumSin = 0.0;
+		double Energy = 0.0;
+
+		for (int32 Sample = 0; Sample < InNumFrames; ++Sample)
+		{
+			const float SampleVal = Buffer[Sample * InNumChannels + ChannelIdx];
+			Energy += SampleVal * SampleVal;
+
+			const float CosVal = FMath::Cos(Omega * Sample);
+			const float SinVal = FMath::Sin(Omega * Sample);
+
+			SumCos += SampleVal * CosVal;
+			SumSin += SampleVal * SinVal;
+		}
+
+		// Scale to get best-fit sinusoid coefficients for x[n] ~ A * sin(omega n + phi)
+		// A*cos(phi) ~ (2/N) SumSin ; A*sin(phi) ~ (2/N) SumCos (depending on convention).
+		// Here derive A and phi consistently:
+		const double Norm = 2.0 / double(InNumFrames);
+		const double A_cosphi = Norm * SumSin; // correlates with sin
+		const double A_sinphi = Norm * SumCos; // correlates with cos
+
+		OutAmplitude = float(FMath::Sqrt(A_cosphi * A_cosphi + A_sinphi * A_sinphi));
+		OutPhase = (UE_PI / 2.0) - float(FMath::Atan2(A_cosphi, A_sinphi)); // so that sin(n*omega + phi) fits
+
+		// No point in testing energy if there is no amplitude
+		if (OutAmplitude < RMSErrorTolerance)
+		{
+			return false;
+		}
+
+		// Compute reconstruction and RMSE
+		double SumSquaredError = 0.0;
+		for (int32 Sample = 0; Sample < InNumFrames; ++Sample)
+		{
+			const float SineVal = OutAmplitude * FMath::Sin(Omega * Sample + OutPhase);
+			const double Error = double(Buffer[Sample * InNumChannels + ChannelIdx]) - double(SineVal);
+			SumSquaredError += Error * Error;
+		}
+
+		const double RMSError = FMath::Sqrt(SumSquaredError / double(InNumFrames));
+
+		// Optional sanity checks on amplitude bounds if you expect unit amplitude, etc.
+		// Return pass if RMSE small enough
+		return RMSError <= RMSErrorTolerance;
+	}
+
+
+private:
+	// Poll task variables
+	UE::Tasks::FPipe Pipe;
+	UE::Tasks::FTask PollTask;
+	FThreadSafeBool bStop{ false };
+
+	// Variables needed to connect to the listener and parameters to format the data as desired
+	FAudioDevice* AudioDevice = nullptr;
+	USoundSubmix* Submix;
+	int32 FrameCount;
+	int32 ChannelCount;
+	int32 SampleRate;
+
+	// Device under test
+	TSharedPtr<FDownmixedBufferedSubmixListener> Listener;
+
+	// Test variables
+	float ExpectedSineFreq;
+	std::atomic<int32> BuffersReceived{ 0 };
+	std::atomic<float> LatestRMS{ 0.0f };
+	std::atomic<int32> SuccessfulTests{ 0 };
+	std::atomic<int32> FailedTests{ 0 };
+	TArray<float> TestBuffer;
+};
+
+TEST_CLASS_WITH_FLAGS(BufferedSubmixListener_SmokeTest, "Audio.Mixer", EAutomationTestFlags::EditorContext | EAutomationTestFlags::EngineFilter)
+{
+	// Test constants
+	static constexpr uint32 SampleRate = 48000;
+	static constexpr float  FrequencyHz = 440.0f;
+	static constexpr float  TargetSendLevel = 1.0f;
+	static constexpr float  MinDetectableRMSForSignal = 0.0001f; 
+
+	static constexpr int32 FrameCount = SampleRate / 100; // VoiceChat wants data in 10ms chunks
+	static constexpr int32 ChannelCount = 2; // Proximity VoiceChat wants a stereo stream
+
+	// Owned objects
+	TObjectPtr<USoundSubmix> TestSubmix = nullptr;
+	TObjectPtr<UTestSineSynthComponent> SineGenerator = nullptr;
+
+	FAudioDevice* AudioDevice = nullptr;
+
+	// Test rig
+	TUniquePtr<FCQTestPoller> VoiceChatSim;
+
+	BEFORE_EACH()
+	{
+		// 1) Get main audio device
+		AudioDevice = GEngine ? GEngine->GetMainAudioDeviceRaw() : nullptr;
+		ASSERT_THAT(IsNotNull(AudioDevice));
+
+		// 2) Create a transient submix instance for routing
+		TestSubmix = NewObject<USoundSubmix>();
+		ASSERT_THAT(IsNotNull(TestSubmix));
+
+		// 3) Create the test sine tone
+		SineGenerator = NewObject<UTestSineSynthComponent>();
+		SineGenerator->FrequencyHz = FrequencyHz;
+		SineGenerator->bAutoActivate = false;
+		SineGenerator->bIsUISound = true;
+		SineGenerator->bAllowSpatialization = false;
+		SineGenerator->CreateAudioComponent();
+		SineGenerator->SetSubmixSend(TestSubmix, TargetSendLevel);
+
+		// Start audio
+		SineGenerator->Start();
+
+		// 4) Create the VoiceChat simulation tester
+		VoiceChatSim = MakeUnique<FCQTestPoller>(AudioDevice, TestSubmix, FrameCount, ChannelCount, SampleRate);
+
+		// 5) Start the Consumer thread
+		VoiceChatSim->ClearSetupValues();
+		VoiceChatSim->SetTestSineFrequency(FrequencyHz);
+		VoiceChatSim->Start();
+	}
+
+	AFTER_EACH()
+	{
+		VoiceChatSim->Stop();
+		VoiceChatSim.Reset();
+
+		if (SineGenerator)
+		{
+			SineGenerator->Stop();
+			SineGenerator = nullptr;
+		}
+
+		TestSubmix = nullptr;
+	}
+
+	// Wait for the listener to receive a few buffers and show non-zero RMS
+	// then listen for a few buffers of data to make sure it is the sine wave we expect
+	TEST_METHOD(PlaysSineToSubmixAndIsObservedInBuffer)
+	{
+		int32 FirstTestBuffer = 0;
+		// Wait until we get a few buffers with non-trivial energy
+		TestCommandBuilder
+			.WaitDelay(FTimespan::FromMilliseconds(1000))
+
+			// Wait until the RMS indicates audible signal present
+			.Until(TEXT("Wait for non-zero RMS"),
+				[this, &FirstTestBuffer]()
+				{
+					FirstTestBuffer = VoiceChatSim->GetBuffersReceived();
+					VoiceChatSim->ClearTestValues();
+					return VoiceChatSim->GetLatestRMS() > MinDetectableRMSForSignal;
+				})
+
+			// Run the test for a number of received buffers
+			.Until(TEXT("Wait for submix listener to receive buffers"),
+				[this, FirstTestBuffer]()
+				{
+					return VoiceChatSim->GetBuffersReceived() > FirstTestBuffer + 10;
+				})
+
+			.Then([this]()
+				{
+					// Final assert: we observed a valid sine wave for every buffer
+					const int32 SuccessfulSineTest = VoiceChatSim->GetSuccessfulTestCount();
+					const int32 FailedSineTest = VoiceChatSim->GetFailedTestCount();
+					ASSERT_THAT(IsTrue(SuccessfulSineTest > 0 && FailedSineTest == 0));
+				});
+	}
+};
+
+#endif // WITH_DEV_AUTOMATION_TESTS
\ No newline at end of file
diff --git a/Engine/Plugins/Tests/AudioCQTests/Source/AudioCQTests/Private/TestSineSynthComponent.h b/Engine/Plugins/Tests/AudioCQTests/Source/AudioCQTests/Private/TestSineSynthComponent.h
new file mode 100644
--- /dev/null
+++ b/Engine/Plugins/Tests/AudioCQTests/Source/AudioCQTests/Private/TestSineSynthComponent.h
@@ -0,0 +1,54 @@
+// Copyright Epic Games, Inc. All Rights Reserved.
+
+#pragma once
+
+#include "Components/SynthComponent.h"
+#include "DSP/SinOsc.h" 
+#include "TestSineSynthComponent.generated.h"
+
+UCLASS()
+class UTestSineSynthComponent : public USynthComponent
+{
+    GENERATED_BODY()
+
+public:
+    // Frequency Hz
+    UPROPERTY(EditAnywhere, Category = "Test|Audio")
+    float FrequencyHz = 440.0f;
+
+    // Amplitude 0..1
+    UPROPERTY(EditAnywhere, Category = "Test|Audio")
+    float Amplitude = 0.25f;
+
+protected:
+    // Called by USynthComponent::Initialize
+    virtual bool Init(int32& SampleRate) override
+    {
+        // Sample rate chosen by USynthComponent will be provided here
+        Osc.Init((float)SampleRate, FrequencyHz);
+		NumChannels = 1;
+        return true;
+    }
+
+    // Called per render block
+    virtual int32 OnGenerateAudio(float* OutAudio, int32 NumSamples) override
+    {
+        // NumSamples accounts for number of channels
+        // Generate mono then upmix to component NumChannels
+        const int32 NumFrames = NumSamples / NumChannels;
+
+        for (int32 Frame = 0; Frame < NumFrames; ++Frame)
+        {
+            // FSineOsc::ProcessAudio() returns one sample at osc amplitude
+            const float Sample = Osc.ProcessAudio() * Amplitude;
+            for (int32 Ch = 0; Ch < NumChannels; ++Ch)
+            {
+                OutAudio[Frame * NumChannels + Ch] = Sample;
+            }
+        }
+        return NumSamples;
+    }
+
+private:
+    Audio::FSineOsc Osc;
+};
\ No newline at end of file
diff --git a/Engine/Source/Runtime/AudioLink/AudioLinkEngine/AudioLinkEngine.Build.cs b/Engine/Source/Runtime/AudioLink/AudioLinkEngine/AudioLinkEngine.Build.cs
--- a/Engine/Source/Runtime/AudioLink/AudioLinkEngine/AudioLinkEngine.Build.cs
+++ b/Engine/Source/Runtime/AudioLink/AudioLinkEngine/AudioLinkEngine.Build.cs
@@ -13,6 +13,7 @@ namespace UnrealBuildTool.Rules
 					"Core",
 					"CoreUObject",
 					"Engine",
+					"SignalProcessing",
 				}
 			);
 		}
diff --git a/Engine/Source/Runtime/AudioLink/AudioLinkEngine/Private/BufferedListenerBase.cpp b/Engine/Source/Runtime/AudioLink/AudioLinkEngine/Private/BufferedListenerBase.cpp
--- a/Engine/Source/Runtime/AudioLink/AudioLinkEngine/Private/BufferedListenerBase.cpp
+++ b/Engine/Source/Runtime/AudioLink/AudioLinkEngine/Private/BufferedListenerBase.cpp
@@ -15,7 +15,7 @@ bool FBufferedListenerBase::PopBuffer(float* InBuffer, int32 InBufferSizeInSampl
 	FScopeLock Lock(&CircularBuffer.GetCriticialSection());
 	OutSamplesWritten = CircularBuffer.Pop(InBuffer, InBufferSizeInSamples);
 	int32 NumSamplesRemaining = CircularBuffer.Num();
-	
+
 	// Return true if any data remaining, which signals the consumer to shutdown if false.
 	// However if we're still running, we could be starving, so only signal stopped if we've exhausted.
 	return (bStopping && NumSamplesRemaining > 0) || bStarted;
@@ -66,7 +66,10 @@ void FBufferedListenerBase::OnBufferReceived(const FBufferFormat& InFormat, TArr
 		else
 		{
 			// Sanity check they haven't changed on the source since it started.
-			ensure(InFormat == *KnownFormat);
+			// Note, in this case it is okay to be off a couple samples if we are doing sample rate conversion	
+			ensure(InFormat.NumChannels == KnownFormat->NumChannels &&
+				FMath::Abs(InFormat.NumSamplesPerBlock - KnownFormat->NumSamplesPerBlock) <= InFormat.NumChannels &&
+				InFormat.NumSamplesPerSec == KnownFormat->NumSamplesPerSec);
 		}
 	}
 
@@ -81,13 +84,13 @@ void FBufferedListenerBase::OnBufferReceived(const FBufferFormat& InFormat, TArr
 	{
 		// Lock outside for the sake of atomic logging
 		FScopeLock Lock(&CircularBuffer.GetCriticialSection());
-		
+
 		// Push the data into the circular buffer.
 		SamplesPushed = CircularBuffer.Push(InBuffer.GetData(), InBuffer.Num());
 
-		UE_LOG(LogAudioLink, VeryVerbose, 
+		UE_LOG(LogAudioLink, VeryVerbose,
 			TEXT("FBufferedListenerBase::OnBufferReceived()(post-push), SamplesPushed=%d, InBuffer.Num()=%d, LocklessCircularBuffer.Num()=%d, LocklessCircularBuffer.Remainder()=%d, LocklessCircularBuffer.Capacity()=%d, This=0x%p"),
-				SamplesPushed, InBuffer.Num(), CircularBuffer.Num(), CircularBuffer.Remainder(), CircularBuffer.GetCapacity(), this );
+			SamplesPushed, InBuffer.Num(), CircularBuffer.Num(), CircularBuffer.Remainder(), CircularBuffer.GetCapacity(), this);
 	}
 
 	// Warn of not enough space in circular buffer, unless we're overwriting silence.
@@ -96,7 +99,7 @@ void FBufferedListenerBase::OnBufferReceived(const FBufferFormat& InFormat, TArr
 		// Prevent log spam by limiting to 1:100 logs.
 		static const int32 NumLogMessagesToSkip = 100;
 		static int32 LogPacifier = 0;
-		UE_CLOG(LogPacifier++ % NumLogMessagesToSkip == 0, LogAudioLink, 
+		UE_CLOG(LogPacifier++ % NumLogMessagesToSkip == 0, LogAudioLink,
 			Verbose, TEXT("FBufferedListenerBase: Overflow by '%d' Samples in Buffer Listener"), InBuffer.Num() - SamplesPushed);
 	}
 }
@@ -149,11 +152,18 @@ void FBufferedListenerBase::PushSilence(int32 InNumSamplesOfSilence)
 void FBufferedListenerBase::Reserve(int32 InNumSamplesToReserve, int32 InNumSamplesOfSilence)
 {
 	// This Zeros the buffer also so should only be done at the start.
-	CircularBuffer.SetCapacity(InNumSamplesToReserve);	
-	
+	CircularBuffer.SetCapacity(InNumSamplesToReserve);
+
 	// Optionally add silence into the buffer, this will cause latency.
 	if (InNumSamplesOfSilence > 0)
 	{
 		PushSilence(InNumSamplesOfSilence);
 	}
 }
+
+int32 FBufferedListenerBase::GetNumBufferedSamples()
+{
+	// Lock for safety
+	FScopeLock Lock(&CircularBuffer.GetCriticialSection());
+	return CircularBuffer.Num();
+}
diff --git a/Engine/Source/Runtime/AudioLink/AudioLinkEngine/Private/BufferedListenerBase.h b/Engine/Source/Runtime/AudioLink/AudioLinkEngine/Private/BufferedListenerBase.h
deleted file mode 100644
diff --git a/Engine/Source/Runtime/AudioLink/AudioLinkEngine/Private/BufferedSourceListener.h b/Engine/Source/Runtime/AudioLink/AudioLinkEngine/Private/BufferedSourceListener.h
deleted file mode 100644
diff --git a/Engine/Source/Runtime/AudioLink/AudioLinkEngine/Private/BufferedSubmixListener.cpp b/Engine/Source/Runtime/AudioLink/AudioLinkEngine/Private/BufferedSubmixListener.cpp
--- a/Engine/Source/Runtime/AudioLink/AudioLinkEngine/Private/BufferedSubmixListener.cpp
+++ b/Engine/Source/Runtime/AudioLink/AudioLinkEngine/Private/BufferedSubmixListener.cpp
@@ -32,37 +32,31 @@ const FString& FBufferedSubmixListener::GetListenerName() const
 	return Name;
 }
 
-void FBufferedSubmixListener::RegisterWithAudioDevice(FAudioDevice* InDevice)
-{
-	check(InDevice);
-
-	DeviceId = InDevice->DeviceID;
-	InDevice->RegisterSubmixBufferListener(AsShared(), InDevice->GetMainSubmixObject());
-}
-
-void FBufferedSubmixListener::UnregsiterWithAudioDevice(FAudioDevice* InDevice)
-{
-	check(InDevice);
-
-	if (ensure(DeviceId != BufferedSubmixListenerPrivate::InvalidAudioDeviceId))
-	{
-		DeviceId = BufferedSubmixListenerPrivate::InvalidAudioDeviceId;
-		InDevice->UnregisterSubmixBufferListener(AsShared(), InDevice->GetMainSubmixObject());
-	}
-}
-
-bool FBufferedSubmixListener::Start(FAudioDevice* InAudioDevice)
+bool FBufferedSubmixListener::Start(FAudioDevice* InAudioDevice, USoundSubmix* InSubmix)
 {
 	if (ensure(InAudioDevice))
 	{
 		if (TrySetStartedFlag())
 		{
-			RegisterWithAudioDevice(InAudioDevice);
+			if (ensure(IsValid(InSubmix)))
+			{
+				InAudioDevice->RegisterSubmixBufferListener(AsShared(), *InSubmix);
+		
+				DeviceId = InAudioDevice->DeviceID;
+				Submix = InSubmix ? TObjectPtr<USoundSubmix>(InSubmix) : TObjectPtr<USoundSubmix>(&InAudioDevice->GetMainSubmixObject());
+				
+				return true;
+			}
 		}
 	}
 	return false;
 }
 
+bool FBufferedSubmixListener::Start(FAudioDevice* InAudioDevice)
+{
+	return Start(InAudioDevice, &InAudioDevice->GetMainSubmixObject());
+}
+
 void FBufferedSubmixListener::Stop(FAudioDevice* InAudioDevice)
 {
 	if (ensure(InAudioDevice))
@@ -71,13 +65,19 @@ void FBufferedSubmixListener::Stop(FAudioDevice* InAudioDevice)
 		{
 			if (TryUnsetStartedFlag())
 			{
-				UnregsiterWithAudioDevice(InAudioDevice);
+				if (ensure(IsValid(Submix)))
+				{
+					InAudioDevice->UnregisterSubmixBufferListener(AsShared(), *Submix);
+				}
+
+				DeviceId = BufferedSubmixListenerPrivate::InvalidAudioDeviceId;
+				Submix = nullptr;
 			}
 		}
 	}
 }
 
-void FBufferedSubmixListener::OnNewSubmixBuffer(const USoundSubmix* OwningSubmix, float* AudioData, int32 InNumSamples, int32 InNumChannels, const int32 InSampleRate, double)
+void FBufferedSubmixListener::OnNewSubmixBuffer(const USoundSubmix* OwningSubmix, float* AudioData, int32 InNumSamples, int32 InNumChannels, const int32 InSampleRate, double /*AudioClock*/)
 {
 	if (IsStartedNonAtomic())
 	{
@@ -95,6 +95,6 @@ void FBufferedSubmixListener::OnNewSubmixBuffer(const USoundSubmix* OwningSubmix
 		if (bZeroInputBuffer)
 		{
 			FMemory::Memzero(AudioData, sizeof(float) * InNumSamples);
-		}	
+		}
 	}
 }
diff --git a/Engine/Source/Runtime/AudioLink/AudioLinkEngine/Private/BufferedSubmixListener.h b/Engine/Source/Runtime/AudioLink/AudioLinkEngine/Private/BufferedSubmixListener.h
deleted file mode 100644
diff --git a/Engine/Source/Runtime/AudioLink/AudioLinkEngine/Private/DownmixedBufferedSubmixListener.cpp b/Engine/Source/Runtime/AudioLink/AudioLinkEngine/Private/DownmixedBufferedSubmixListener.cpp
new file mode 100644
--- /dev/null
+++ b/Engine/Source/Runtime/AudioLink/AudioLinkEngine/Private/DownmixedBufferedSubmixListener.cpp
@@ -0,0 +1,115 @@
+// Copyright Epic Games, Inc. All Rights Reserved.
+
+#include "DownmixedBufferedSubmixListener.h"
+#include "DSP/BufferVectorOperations.h"
+#include "DSP/ChannelMap.h"
+#include "SampleBuffer.h"
+#include "Sound/SoundSubmix.h"
+
+namespace DownmixedBufferedSubmixListenerPrivate
+{
+	constexpr int32 MaxInputChannels = 8;
+	constexpr int32 MaxNumFramesPerBuffer = 2048;
+}
+
+/** Buffered Submix Listener with downmixing and a delegate callback */
+FDownmixedBufferedSubmixListener::FDownmixedBufferedSubmixListener(int32 InNumOutputChannels, int32 InSampleRate, const FString* InName)
+	: FBufferedSubmixListener{ InNumOutputChannels * InSampleRate, false, InName }
+	, NumOutputChannels(InNumOutputChannels)
+	, OutputSampleRate(InSampleRate)
+	, Resampler(InNumOutputChannels)
+{
+	ChannelGainMap.Reserve(InNumOutputChannels * DownmixedBufferedSubmixListenerPrivate::MaxInputChannels);
+	DownMixBuffer.Reserve(InNumOutputChannels * DownmixedBufferedSubmixListenerPrivate::MaxNumFramesPerBuffer);
+	RateConvertedBuffer.Reserve(InNumOutputChannels * DownmixedBufferedSubmixListenerPrivate::MaxNumFramesPerBuffer);
+}
+
+int32 FDownmixedBufferedSubmixListener::GetNumAvailableSamples()
+{
+	return GetNumBufferedSamples();
+}
+
+bool FDownmixedBufferedSubmixListener::GetBuffer(float* InBuffer, int32 InBufferSizeInSamples, int32& OutSamplesWritten)
+{
+	return PopBuffer(InBuffer, InBufferSizeInSamples, OutSamplesWritten);
+}
+
+void FDownmixedBufferedSubmixListener::OnNewSubmixBuffer(const USoundSubmix* OwningSubmix, float* AudioData, int32 InNumSamples, int32 InNumChannels,
+	const int32 InSampleRate, double /*AudioClock*/)
+{
+	using namespace Audio;
+
+	if (InNumSamples <= 0 || InNumChannels <= 0 || OutputSampleRate <= 0)
+	{
+		return;
+	}
+
+	if (IsStartedNonAtomic())
+	{
+		// Rebuild the sample rate converter if the mixer sample rate changed
+		if (InSampleRate != InputSampleRate)
+		{
+			Resampler.SetFrameRatio((float)InSampleRate / OutputSampleRate);
+			InputSampleRate = InSampleRate;
+		}
+
+		// Rebuild the downmix gain map if the number of mixer channels has changed
+		if (InNumChannels != NumInputChannels)
+		{
+			check(InNumChannels <= DownmixedBufferedSubmixListenerPrivate::MaxInputChannels);
+
+			FChannelMapParams GainMapParams;
+			GainMapParams.NumInputChannels = InNumChannels;
+			GainMapParams.NumOutputChannels = NumOutputChannels;
+			GainMapParams.Order = EChannelMapOrder::OutputMajorOrder;
+			GainMapParams.MonoUpmixMethod = EChannelMapMonoUpmixMethod::EqualPower;
+			GainMapParams.bIsCenterChannelOnly = false;
+
+			Create2DChannelMap(GainMapParams, ChannelGainMap);
+			NumInputChannels = InNumChannels;
+		}
+
+		int32 NumInputFrames = InNumSamples / InNumChannels;
+		check(NumInputFrames <= DownmixedBufferedSubmixListenerPrivate::MaxNumFramesPerBuffer);
+
+		const float* DownmixerOutputBuffer = AudioData;
+
+		// Downmix the mixer data if necessary
+		if (InNumChannels != NumOutputChannels)
+		{
+			DownMixBuffer.SetNumUninitialized(NumInputFrames * NumOutputChannels);
+			DownmixBuffer(InNumChannels, NumOutputChannels, AudioData, DownMixBuffer.GetData(), NumInputFrames, ChannelGainMap.GetData());
+			DownmixerOutputBuffer = DownMixBuffer.GetData();
+		}
+
+		// Convert the sample rate if necessary
+		int32 NumOutputFrames = NumInputFrames;
+		const float* OutputBuffer = DownmixerOutputBuffer;
+		if (OutputSampleRate != InSampleRate)
+		{
+			NumOutputFrames = Resampler.GetNumOutputFramesProducedByInputFrames(NumInputFrames);
+			RateConvertedBuffer.SetNumUninitialized(NumOutputFrames * NumOutputChannels);
+
+			int32 NumFramesConsumed = -1;
+			int32 NumFramesProduced = -1;
+
+			Resampler.ProcessInterleaved(TArrayView<const float>{DownmixerOutputBuffer, NumInputFrames * NumOutputChannels},
+										 TArrayView<float>{RateConvertedBuffer.GetData(), NumOutputFrames * NumOutputChannels},
+										 NumFramesConsumed, NumFramesProduced);
+
+			OutputBuffer = RateConvertedBuffer.GetData();
+			NumOutputFrames = NumFramesProduced;
+		}
+
+		// Call to base class to handle.
+		FBufferFormat NewFormat;
+		NewFormat.NumChannels = NumOutputChannels;
+		NewFormat.NumSamplesPerBlock = NumOutputFrames * NumOutputChannels;
+		NewFormat.NumSamplesPerSec = OutputSampleRate;
+		OnBufferReceived(NewFormat, MakeArrayView(OutputBuffer, NumOutputFrames * NumOutputChannels));
+
+		// Broadcast to delegates that a buffer has been received
+		OnSubmixBufferWritten.ExecuteIfBound(OutputBuffer, NumOutputFrames, NumOutputChannels);
+	}
+}
+
diff --git a/Engine/Source/Runtime/AudioLink/AudioLinkEngine/Public/BufferedListenerBase.h b/Engine/Source/Runtime/AudioLink/AudioLinkEngine/Public/BufferedListenerBase.h
new file mode 100644
--- /dev/null
+++ b/Engine/Source/Runtime/AudioLink/AudioLinkEngine/Public/BufferedListenerBase.h
@@ -0,0 +1,172 @@
+// Copyright Epic Games, Inc. All Rights Reserved.
+
+#pragma once
+
+#include "AudioDevice.h"
+#include "DSP/Dsp.h"
+
+#include "IBufferedAudioOutput.h"
+#include "Misc/Optional.h"
+#include "Math/UnrealMathUtility.h"
+
+// Simple locking circular buffer
+// Uses 32bit wrap around logic explained here: 
+// https://www.snellman.net/blog/archive/2016-12-13-ring-buffers/
+class FLockingCircularSampleBuffer
+{
+public:
+	FLockingCircularSampleBuffer(int32 InInitialCapacity)
+	{
+		SetCapacity(InInitialCapacity);
+	}
+
+	// Allow the caller to lock the CS ahead of calling (if necessary)
+	FCriticalSection& GetCriticialSection() { return CS; }
+	const FCriticalSection& GetCriticialSection() const { return CS; }
+
+	// This doesn't preserve contents.
+	void SetCapacity(int32 InCapacity)
+	{
+		check(InCapacity > 1);
+
+		if (!FMath::IsPowerOfTwo(InCapacity))
+		{
+			InCapacity = FMath::RoundUpToPowerOfTwo(InCapacity);
+		}
+
+		FScopeLock Lock(&CS);
+		Buffer.SetNumZeroed(InCapacity, EAllowShrinking::Yes);
+		Mask = InCapacity - 1;
+		Read = 0; // Empty.
+		Write = 0;
+
+		check(Mask != 0);
+	}
+	int32 GetCapacity() const
+	{
+		FScopeLock Lock(&CS);
+		return Buffer.Num();
+	}
+	int32 Num() const
+	{
+		FScopeLock Lock(&CS);
+		return Write - Read;
+	}
+
+	// Get number of samples that can be pushed onto the buffer before it is full.
+	int32 Remainder() const
+	{
+		FScopeLock Lock(&CS);
+		return Buffer.Num() - Num();
+	}
+
+	int32 Push(const float* InBuffer, int32 InSize)
+	{
+		int32 CanPush = FMath::Min(Remainder(), InSize);
+		FScopeLock Lock(&CS);
+		for (int32 i = 0; i < CanPush; ++i)
+		{
+			Enqueue(InBuffer[i]);
+		}
+		return CanPush;
+	}
+	int32 Pop(float* OutBuffer, int32 InNumSamples)
+	{
+		int32 CanPop = FMath::Min(Num(), InNumSamples);
+		FScopeLock Lock(&CS);
+		for (int32 i = 0; i < CanPop; ++i)
+		{
+			OutBuffer[i] = Dequeue();
+		}
+		return CanPop;
+	}
+
+	int32 PushZeros(int32 InNumSamplesOfSilence)
+	{
+		int32 CanPush = FMath::Min(Remainder(), InNumSamplesOfSilence);
+		FScopeLock Lock(&CS);
+		for (int32 i = 0; i < CanPush; ++i)
+		{
+			Enqueue(0.f);
+		}
+		return CanPush;
+	}
+private:
+	uint32 Read = 0;		// These grow indefinitely until wrap at 2^32,
+	uint32 Write = 0;		// this allows us to use full capacity as write >= read.
+	uint32 Mask = 0;
+	TArray<float> Buffer;
+	mutable FCriticalSection CS;
+
+	// NOTE: Not forceinline as compiler does better job without.
+	void Enqueue(const float InFloat)
+	{
+		Buffer[Write++ & Mask] = InFloat;
+	}
+	float Dequeue()
+	{
+		return Buffer[Read++ & Mask];
+	}
+};
+
+/** Common base class of Buffered Listener objects.
+*/
+class FBufferedListenerBase : public IBufferedAudioOutput
+{
+protected:
+	AUDIOLINKENGINE_API FBufferedListenerBase(int32 InDefaultCircularBufferSize);
+	virtual ~FBufferedListenerBase() = default;
+
+	//~ Begin IBufferedAudioOutput
+	AUDIOLINKENGINE_API bool PopBuffer(float* InBuffer, int32 InBufferSizeInSamples, int32& OutSamplesWritten) override;
+	AUDIOLINKENGINE_API bool GetFormat(IBufferedAudioOutput::FBufferFormat& OutFormat) const override;
+	AUDIOLINKENGINE_API void Reserve(int32 InNumSamplesToReserve, int32 InNumSamplesOfSilence) override;
+	AUDIOLINKENGINE_API void SetFormatKnownDelegate(FOnFormatKnown InFormatKnownDelegate) override;
+	void SetBufferStreamEndDelegate(FOnBufferStreamEnd) override {}
+	//~ End IBufferedAudioOutput
+
+	//* Common path to receive a new buffer, call from derived classes */
+	AUDIOLINKENGINE_API void OnBufferReceived(const FBufferFormat& InFormat, TArrayView<const float> InBuffer);
+
+	//* Reset the format of the buffer */
+	AUDIOLINKENGINE_API void ResetFormat();
+
+	//* Set the format of the buffer */
+	AUDIOLINKENGINE_API void SetFormat(const FBufferFormat& InFormat);
+
+	//* Ask if the started flag has been set. Note this is non-atomic, as it could change during the call */	
+	AUDIOLINKENGINE_API bool IsStartedNonAtomic() const;
+
+	//* Attempt to set our state to started. Not this can fail if we're already started. */	
+	AUDIOLINKENGINE_API bool TrySetStartedFlag();
+
+	//* Attempt to set started to false. This can fail if we're already stopped. */	
+	AUDIOLINKENGINE_API bool TryUnsetStartedFlag();
+
+	AUDIOLINKENGINE_API bool TrySetStoppingFlag();
+
+	//* Return the number of samples able to be popped off the queue */
+	AUDIOLINKENGINE_API int32 GetNumBufferedSamples();
+
+private:
+	AUDIOLINKENGINE_API void PushSilence(int32 InNumSamplesOfSilence);
+
+	/** Buffer to hold the data for the single source we're listening to, interleaved. */
+	FLockingCircularSampleBuffer CircularBuffer;
+
+	/** Read/Write slim lock protects format known optional */
+	mutable FRWLock FormatKnownRwLock;
+
+	/** Optional that holds the buffer format, if (and when) it's known. Protected by r/w Slim-lock */
+	TOptional<FBufferFormat> KnownFormat;
+
+	/** Delegate that fires when the format it known. Normally on the first buffer received. */
+	FOnFormatKnown OnFormatKnown;
+
+	/** Atomic flag we've been started */
+	std::atomic<bool> bStarted{ false };
+
+	/** Atomic flag we've been told to stop */
+	std::atomic<bool> bStopping{ false };
+};
+
diff --git a/Engine/Source/Runtime/AudioLink/AudioLinkEngine/Public/BufferedSourceListener.h b/Engine/Source/Runtime/AudioLink/AudioLinkEngine/Public/BufferedSourceListener.h
new file mode 100644
--- /dev/null
+++ b/Engine/Source/Runtime/AudioLink/AudioLinkEngine/Public/BufferedSourceListener.h
@@ -0,0 +1,40 @@
+// Copyright Epic Games, Inc. All Rights Reserved.
+
+#pragma once
+
+#include "AudioDevice.h"
+#include "BufferedListenerBase.h"
+
+/**
+	FBufferedSourceListener: Concrete implementation of both a buffer listener and buffered output.
+	Contains a circular buffer.
+*/
+class FBufferedSourceListener : public ISourceBufferListener, public FBufferedListenerBase
+{
+public:
+	AUDIOLINKENGINE_API FBufferedSourceListener(int32 InDefaultCircularBufferSize);
+	AUDIOLINKENGINE_API virtual ~FBufferedSourceListener();
+
+	/**
+	 * @param  InDevice Audio Device to register this listener with.
+	 *
+	 * @return  success true, false otherwise.
+	 */
+	AUDIOLINKENGINE_API bool Start(FAudioDevice* InDevice) override;
+
+	/** Stop the buffer by unregistering the buffer listener */
+	AUDIOLINKENGINE_API void Stop(FAudioDevice* InDevice) override;
+
+private:
+	//~ Begin ISourceBufferListener
+	AUDIOLINKENGINE_API void OnNewBuffer(const ISourceBufferListener::FOnNewBufferParams& InParams) override;
+	AUDIOLINKENGINE_API void OnSourceReleased(const int32 SourceId) override;
+	//~ End ISourceBufferListener
+
+	//~ Begin IBufferedAudioOutput
+	AUDIOLINKENGINE_API void SetBufferStreamEndDelegate(FOnBufferStreamEnd InBufferStreamEndDelegate) override;
+	//~ End IBufferedAudioOutput
+
+	std::atomic<int32> CurrentSourceId = INDEX_NONE;		// r/w AudioMixer thread.	
+	FOnBufferStreamEnd OnBufferStreamEndDelegate;
+};
diff --git a/Engine/Source/Runtime/AudioLink/AudioLinkEngine/Public/BufferedSubmixListener.h b/Engine/Source/Runtime/AudioLink/AudioLinkEngine/Public/BufferedSubmixListener.h
new file mode 100644
--- /dev/null
+++ b/Engine/Source/Runtime/AudioLink/AudioLinkEngine/Public/BufferedSubmixListener.h
@@ -0,0 +1,48 @@
+// Copyright Epic Games, Inc. All Rights Reserved.
+
+#pragma once
+
+#include "BufferedListenerBase.h"
+#include "ISubmixBufferListener.h"
+#include "Sound/SoundSubmix.h"
+
+/** Concrete Submix Buffer Listener
+*/
+class FBufferedSubmixListener : public ISubmixBufferListener, public FBufferedListenerBase
+{
+public:
+	/** Constructor
+	 * @param  InDefaultCircularBufferSize: Size of the circular buffer in samples by default.
+	 * @param  bInZeroInputBuffer: If set, AudioData input buffer will be zeroed out after data is queued
+	 * @param  InName(Optional): Optional name to track listener lifetime with.
+	 */
+	AUDIOLINKENGINE_API FBufferedSubmixListener(int32 InDefaultCircularBufferSize, bool bInZeroInputBuffer, const FString* InName);
+
+	AUDIOLINKENGINE_API virtual ~FBufferedSubmixListener();
+
+	/**
+	 * Starts the Submix buffer listener by registering it with the passed in Audio Device.
+	 * @param  InDevice: Audio Device to register this submix listener with.
+	 * @param  InSubmix: The submix being listened to (default is main output)
+	 *
+	 * @return  success true, false otherwise.
+	 */
+	AUDIOLINKENGINE_API bool Start(FAudioDevice* InDevice, USoundSubmix* InSubmix);
+	//~ Begin IBufferedAudioOutput
+	AUDIOLINKENGINE_API bool Start(FAudioDevice* InDevice) override;
+	AUDIOLINKENGINE_API void Stop(FAudioDevice* InDevice) override;
+	//~ End IBufferedAudioOutput
+
+protected:
+	//~ Begin ISubmixBufferListener
+	AUDIOLINKENGINE_API void OnNewSubmixBuffer(const USoundSubmix* OwningSubmix, float* AudioData, int32 InNumSamples, int32 InNumChannels, const int32 InSampleRate, double) override;
+	AUDIOLINKENGINE_API const FString& GetListenerName() const override;
+	//~ End ISubmixBufferListener
+
+private:
+	Audio::FDeviceId DeviceId;
+	bool bZeroInputBuffer = false;
+	FString Name;
+
+	TObjectPtr<USoundSubmix> Submix;
+};
diff --git a/Engine/Source/Runtime/AudioLink/AudioLinkEngine/Public/DownmixedBufferedSubmixListener.h b/Engine/Source/Runtime/AudioLink/AudioLinkEngine/Public/DownmixedBufferedSubmixListener.h
new file mode 100644
--- /dev/null
+++ b/Engine/Source/Runtime/AudioLink/AudioLinkEngine/Public/DownmixedBufferedSubmixListener.h
@@ -0,0 +1,55 @@
+// Copyright Epic Games, Inc. All Rights Reserved.
+
+#pragma once
+
+#include "AudioDevice.h"
+#include "BufferedSubmixListener.h"
+#include "DSP/RuntimeResampler.h"
+
+class USoundSubmix;
+
+class FDownmixedBufferedSubmixListener : public FBufferedSubmixListener
+{
+public:
+	/** Constructor
+	 * @param  InNumOutputChannels: channel config for output
+	 * @param  InSampleRate: Desired sample rate
+	 * @param  InName(Optional): Optional name to track listener lifetime with.
+	 */
+	AUDIOLINKENGINE_API FDownmixedBufferedSubmixListener(int32 InNumOutputChannels, int32 InSampleRate, const FString* InName);
+
+	virtual ~FDownmixedBufferedSubmixListener() {}
+
+	// Delegate to signal to the consumer that data is ready
+	DECLARE_DELEGATE_ThreeParams(FOnSubmixBufferWritten, const float* /*SampleBuffer*/, int32 /*FrameCount*/, int32 /*NumChannels*/);
+	FOnSubmixBufferWritten OnSubmixBufferWritten;
+
+	/** Query the number of samples waiting in the queue
+	 * @return total number of samples, e.g. for 480 frames of 2 channels, would return 960
+	 */
+	AUDIOLINKENGINE_API int32 GetNumAvailableSamples();
+
+	/** Retrieve samples from the queue
+	 * @param  InBuffer: memory provided by the caller to copy samples into
+	 * @param  InBufferSizeInSamples: total number of samples to be retrieved, frame count * channel count
+	 * @param  OutSamplesWritten: total number of samples written into the provided buffer
+	 * 
+	 * @return if the queue is still active, it will return true, otherwise, return indicates if there are more samples in the queue
+	 */
+	AUDIOLINKENGINE_API bool GetBuffer(float* InBuffer, int32 InBufferSizeInSamples, int32& OutSamplesWritten);
+
+private:
+	//~ Begin ISubmixBufferListener
+	AUDIOLINKENGINE_API void OnNewSubmixBuffer(const USoundSubmix* OwningSubmix, float* AudioData, int32 InNumSamples, int32 InNumChannels, const int32 InSampleRate, double) override;
+	//~ End ISubmixBufferListener
+
+	int32 NumOutputChannels;
+	int32 OutputSampleRate;
+
+	TArray<float> ChannelGainMap;
+	Audio::FAlignedFloatBuffer DownMixBuffer;
+	Audio::FAlignedFloatBuffer RateConvertedBuffer;
+	Audio::FRuntimeResampler Resampler;
+	int32 NumInputChannels = 0;
+	int32 InputSampleRate = 0;
+};
