Fix BufferedSubmixListener test break on Linux system with stack use after return error.
On 39.20

#Synced-CL 48857083

diff --git a/Engine/Plugins/Tests/AudioCQTests/Source/AudioCQTests/Private/BufferedSubmixListenerTests.cpp b/Engine/Plugins/Tests/AudioCQTests/Source/AudioCQTests/Private/BufferedSubmixListenerTests.cpp
--- a/Engine/Plugins/Tests/AudioCQTests/Source/AudioCQTests/Private/BufferedSubmixListenerTests.cpp
+++ b/Engine/Plugins/Tests/AudioCQTests/Source/AudioCQTests/Private/BufferedSubmixListenerTests.cpp
@@ -250,6 +250,7 @@ TEST_CLASS_WITH_FLAGS(BufferedSubmixListener_SmokeTest, "Audio.Mixer", EAutomati
 
 	// Test rig
 	TUniquePtr<FCQTestPoller> VoiceChatSim;
+	int32 FirstTestBuffer = 0;
 
 	BEFORE_EACH()
 	{
@@ -300,14 +301,13 @@ TEST_CLASS_WITH_FLAGS(BufferedSubmixListener_SmokeTest, "Audio.Mixer", EAutomati
 	// then listen for a few buffers of data to make sure it is the sine wave we expect
 	TEST_METHOD(PlaysSineToSubmixAndIsObservedInBuffer)
 	{
-		int32 FirstTestBuffer = 0;
 		// Wait until we get a few buffers with non-trivial energy
 		TestCommandBuilder
 			.WaitDelay(FTimespan::FromMilliseconds(1000))
 
 			// Wait until the RMS indicates audible signal present
 			.Until(TEXT("Wait for non-zero RMS"),
-				[this, &FirstTestBuffer]()
+				[this]()
 				{
 					FirstTestBuffer = VoiceChatSim->GetBuffersReceived();
 					VoiceChatSim->ClearTestValues();
@@ -316,7 +316,7 @@ TEST_CLASS_WITH_FLAGS(BufferedSubmixListener_SmokeTest, "Audio.Mixer", EAutomati
 
 			// Run the test for a number of received buffers
 			.Until(TEXT("Wait for submix listener to receive buffers"),
-				[this, FirstTestBuffer]()
+				[this]()
 				{
 					return VoiceChatSim->GetBuffersReceived() > FirstTestBuffer + 10;
 				})
