Adding signals to main submix output for more complete testing

#Synced-CL 49553638

diff --git a/Engine/Plugins/Tests/AudioCQTests/Source/AudioCQTests/Private/BufferedSubmixListenerTests.cpp b/Engine/Plugins/Tests/AudioCQTests/Source/AudioCQTests/Private/BufferedSubmixListenerTests.cpp
--- a/Engine/Plugins/Tests/AudioCQTests/Source/AudioCQTests/Private/BufferedSubmixListenerTests.cpp
+++ b/Engine/Plugins/Tests/AudioCQTests/Source/AudioCQTests/Private/BufferedSubmixListenerTests.cpp
@@ -146,7 +146,7 @@ private:
 	// Returns true if buffer is well-fit by a sine at ExpectedFreq within tolerances.
 	// Also outputs estimated amplitude and phase.
 	bool IsSineAtFrequency(const float* Buffer, int32 InNumFrames, float InSampleRate, int32 InNumChannels, int32 ChannelIdx,
-		float ExpectedFreq, float& OutAmplitude, float& OutPhase, float RMSErrorTolerance = 1e-2f)
+		float ExpectedFreq, float& OutAmplitude, float& OutPhase, float RMSErrorTolerance = 0.005f)
 	{
 		if (!Buffer || InNumFrames <= 0 || InSampleRate <= 0.f || ExpectedFreq <= 0.f || InNumChannels <= 0 || ChannelIdx >= InNumChannels)
 		{
@@ -243,8 +243,9 @@ TEST_CLASS_WITH_FLAGS(BufferedSubmixListener_SmokeTest, "Audio.Mixer", EAutomati
 	static constexpr int32 ChannelCount = 2; // Proximity VoiceChat wants a stereo stream
 
 	// Owned objects
-	TObjectPtr<USoundSubmix> TestSubmix = nullptr;
-	TObjectPtr<UTestSineSynthComponent> SineGenerator = nullptr;
+	TStrongObjectPtr<USoundSubmix> TestSubmix = nullptr;
+	TStrongObjectPtr<UTestSineSynthComponent> SineGenerator = nullptr;
+	TStrongObjectPtr<UTestSineSynthComponent> InterferenceTone = nullptr;
 
 	FAudioDevice* AudioDevice = nullptr;
 
@@ -259,23 +260,43 @@ TEST_CLASS_WITH_FLAGS(BufferedSubmixListener_SmokeTest, "Audio.Mixer", EAutomati
 		ASSERT_THAT(IsNotNull(AudioDevice));
 
 		// 2) Create a transient submix instance for routing
-		TestSubmix = NewObject<USoundSubmix>();
-		ASSERT_THAT(IsNotNull(TestSubmix));
+		TestSubmix = TStrongObjectPtr<USoundSubmix>(NewObject<USoundSubmix>());
+		if (AudioDevice)
+		{
+			AudioDevice->RegisterSoundSubmix(TestSubmix.Get(), true);
+		}
+		ASSERT_THAT(IsNotNull(TestSubmix.Get()));
 
 		// 3) Create the test sine tone
-		SineGenerator = NewObject<UTestSineSynthComponent>();
+		SineGenerator = TStrongObjectPtr<UTestSineSynthComponent>(NewObject<UTestSineSynthComponent>());
 		SineGenerator->FrequencyHz = FrequencyHz;
 		SineGenerator->bAutoActivate = false;
 		SineGenerator->bIsUISound = true;
 		SineGenerator->bAllowSpatialization = false;
 		SineGenerator->CreateAudioComponent();
-		SineGenerator->SetSubmixSend(TestSubmix, TargetSendLevel);
+		SineGenerator->SetSubmixSend(TestSubmix.Get(), TargetSendLevel);
 
 		// Start audio
 		SineGenerator->Start();
 
+		// 4) Add an interference tone that sends stuff directly to the final submix just to make sure we aren't defaulting to using main
+		InterferenceTone = TStrongObjectPtr<UTestSineSynthComponent>(NewObject<UTestSineSynthComponent>());
+		InterferenceTone->FrequencyHz = FrequencyHz * 0.707f;
+		InterferenceTone->bAutoActivate = false;
+		InterferenceTone->bIsUISound = true;
+		InterferenceTone->bAllowSpatialization = false;
+		InterferenceTone->CreateAudioComponent();
+
+		if (AudioDevice)
+		{
+			InterferenceTone->SetSubmixSend(&AudioDevice->GetMainSubmixObject(), TargetSendLevel);
+		}
+
+		// Start audio
+		InterferenceTone->Start();
+
 		// 4) Create the VoiceChat simulation tester
-		VoiceChatSim = MakeUnique<FCQTestPoller>(AudioDevice, TestSubmix, FrameCount, ChannelCount, SampleRate);
+		VoiceChatSim = MakeUnique<FCQTestPoller>(AudioDevice, TestSubmix.Get(), FrameCount, ChannelCount, SampleRate);
 
 		// 5) Start the Consumer thread
 		VoiceChatSim->ClearSetupValues();
@@ -294,6 +315,17 @@ TEST_CLASS_WITH_FLAGS(BufferedSubmixListener_SmokeTest, "Audio.Mixer", EAutomati
 			SineGenerator = nullptr;
 		}
 
+		if (InterferenceTone)
+		{
+			InterferenceTone->Stop();
+			InterferenceTone = nullptr;
+		}
+
+		if (AudioDevice)
+		{
+			AudioDevice->UnregisterSoundSubmix(TestSubmix.Get(), false);
+		}
+
 		TestSubmix = nullptr;
 	}
 
