[AudioLink] Add SetOutputFormat() to FDownmixedBufferedSubmixListener for runtime output format changes

Previously, the output sample rate and channel count were fixed at construction time. Consumers
that needed to respond to device format changes (e.g., audio device swaps) had to destroy and
recreate the listener, requiring game-thread async hops and causing dropped audio during the
transition.

SetOutputFormat() allows consumers to change the output format in-place from any thread. It
flushes the circular buffer immediately so consumers never read stale old-format data, then
signals the audio render thread to rebuild the resampler and downmix gain map on the next
callback. When called from the audio render thread (the expected pattern for voice chat),
the transition is seamless with zero underrun.

Engine changes:
- FLockingCircularSampleBuffer: Add Flush() (reset pointers, no realloc) and EnsureCapacity()
(grow-only, never shrinks, consistent with EAllowShrinking::No convention in AudioMixer)
- FBufferedListenerBase: Add FlushAndResetFormat() and FlushAndEnsureCapacity() to flush the
circular buffer and reset KnownFormat so OnFormatKnown re-fires on the next buffer
- FDownmixedBufferedSubmixListener: Add SetOutputFormat() with atomic pending format handoff
and pickup logic in OnNewSubmixBuffer(). Output channel count capped at 8.

Test:
- New CQTest ChangeOutputFormatWhileRunning: validates sine wave integrity at stereo/48kHz,
switches to mono/24kHz mid-stream via SetOutputFormat(), validates sine wave at new format

#Synced-CL 51831414

diff --git a/Engine/Plugins/Tests/AudioCQTests/Source/AudioCQTests/Private/BufferedSubmixListenerTests.cpp b/Engine/Plugins/Tests/AudioCQTests/Source/AudioCQTests/Private/BufferedSubmixListenerTests.cpp
--- a/Engine/Plugins/Tests/AudioCQTests/Source/AudioCQTests/Private/BufferedSubmixListenerTests.cpp
+++ b/Engine/Plugins/Tests/AudioCQTests/Source/AudioCQTests/Private/BufferedSubmixListenerTests.cpp
@@ -102,6 +102,21 @@ public:
 		LatestRMS.store(Rms);
 	}
 
+	// Change the output format of the listener and update the poller's local format to match
+	void SetOutputFormat(int32 InChannelCount, int32 InSampleRate)
+	{
+		if (Listener.IsValid())
+		{
+			Listener->SetOutputFormat(InChannelCount, InSampleRate);
+		}
+
+		const int32 NewFrameCount = InSampleRate / 100;
+		FrameCount.store(NewFrameCount);
+		ChannelCount.store(InChannelCount);
+		SampleRate.store(InSampleRate);
+		TestBuffer.SetNumUninitialized(NewFrameCount * InChannelCount);
+	}
+
 	// Functions to access atomic values during the test process
 	void SetTestSineFrequency(float ExpectedFreq) { ExpectedSineFreq = ExpectedFreq; }
 	int32 GetBuffersReceived() { return BuffersReceived.load(); }
@@ -116,20 +131,25 @@ private:
 	// This fetches a 10ms buffer of data from the BufferedSubmixListener and checks if it is actually a sine wave
 	void PollListenerQueue()
 	{
-		if (Listener.IsValid() && Listener->GetNumAvailableSamples() >= FrameCount * ChannelCount)
+		const int32 CurrentFrameCount = FrameCount.load();
+		const int32 CurrentChannelCount = ChannelCount.load();
+		const int32 CurrentSampleRate = SampleRate.load();
+		const int32 RequiredSamples = CurrentFrameCount * CurrentChannelCount;
+
+		if (Listener.IsValid() && Listener->GetNumAvailableSamples() >= RequiredSamples)
 		{
-			TestBuffer.SetNumUninitialized(FrameCount * ChannelCount);
+			TestBuffer.SetNumUninitialized(RequiredSamples);
 			float* FilledBuffer = TestBuffer.GetData();
 			int32 SamplesPopped = 0;
-			Listener->GetBuffer(TestBuffer.GetData(), FrameCount * ChannelCount, SamplesPopped);
+			Listener->GetBuffer(FilledBuffer, RequiredSamples, SamplesPopped);
 
-			bool TestPasses = ChannelCount > 0;
+			bool TestPasses = CurrentChannelCount > 0;
 
-			for (int32 Channel = 0; Channel < ChannelCount; Channel++)
+			for (int32 Channel = 0; Channel < CurrentChannelCount; Channel++)
 			{
 				float Amplitude = 0.0f;
 				float Phase = 0.0f;
-				TestPasses &= IsSineAtFrequency(FilledBuffer, SamplesPopped / ChannelCount, SampleRate, ChannelCount, Channel, ExpectedSineFreq, Amplitude, Phase);
+				TestPasses &= IsSineAtFrequency(FilledBuffer, SamplesPopped / CurrentChannelCount, CurrentSampleRate, CurrentChannelCount, Channel, ExpectedSineFreq, Amplitude, Phase);
 			}
 
 			if (TestPasses)
@@ -215,9 +235,9 @@ private:
 	// Variables needed to connect to the listener and parameters to format the data as desired
 	FAudioDevice* AudioDevice = nullptr;
 	USoundSubmix* Submix;
-	int32 FrameCount;
-	int32 ChannelCount;
-	int32 SampleRate;
+	std::atomic<int32> FrameCount;
+	std::atomic<int32> ChannelCount;
+	std::atomic<int32> SampleRate;
 
 	// Device under test
 	TSharedPtr<FDownmixedBufferedSubmixListener> Listener;
@@ -361,6 +381,72 @@ TEST_CLASS_WITH_FLAGS(BufferedSubmixListener_SmokeTest, "Audio.Mixer", EAutomati
 					ASSERT_THAT(IsTrue(SuccessfulSineTest > 0 && FailedSineTest == 0));
 				});
 	}
+
+	// Verify that changing the output format mid-stream via SetOutputFormat() produces
+	// valid sine wave data at the new format after the transition
+	TEST_METHOD(ChangeOutputFormatWhileRunning)
+	{
+		static constexpr int32 NewChannelCount = 1;
+		static constexpr uint32 NewSampleRate = 24000;
+
+		TestCommandBuilder
+			.WaitDelay(FTimespan::FromMilliseconds(1000))
+
+			// Phase 1: Wait for valid signal at the original format (stereo 48kHz)
+			.Until(TEXT("Wait for non-zero RMS at original format"),
+				[this]()
+				{
+					FirstTestBuffer = VoiceChatSim->GetBuffersReceived();
+					VoiceChatSim->ClearTestValues();
+					return VoiceChatSim->GetLatestRMS() > MinDetectableRMSForSignal;
+				})
+
+			.Until(TEXT("Validate sine at original format"),
+				[this]()
+				{
+					return VoiceChatSim->GetBuffersReceived() > FirstTestBuffer + 10;
+				})
+
+			.Then([this]()
+				{
+					// Confirm original format works
+					const int32 SuccessfulSineTest = VoiceChatSim->GetSuccessfulTestCount();
+					const int32 FailedSineTest = VoiceChatSim->GetFailedTestCount();
+					ASSERT_THAT(IsTrue(SuccessfulSineTest > 0 && FailedSineTest == 0));
+				})
+
+			// Phase 2: Switch to mono 24kHz and validate
+			.Then([this]()
+				{
+					VoiceChatSim->SetOutputFormat(NewChannelCount, NewSampleRate);
+					VoiceChatSim->ClearSetupValues();
+					VoiceChatSim->ClearTestValues();
+				})
+
+			.WaitDelay(FTimespan::FromMilliseconds(500))
+
+			.Until(TEXT("Wait for non-zero RMS at new format"),
+				[this]()
+				{
+					FirstTestBuffer = VoiceChatSim->GetBuffersReceived();
+					VoiceChatSim->ClearTestValues();
+					return VoiceChatSim->GetLatestRMS() > MinDetectableRMSForSignal;
+				})
+
+			.Until(TEXT("Validate sine at new format"),
+				[this]()
+				{
+					return VoiceChatSim->GetBuffersReceived() > FirstTestBuffer + 10;
+				})
+
+			.Then([this]()
+				{
+					// Confirm new format produces valid sine data
+					const int32 SuccessfulSineTest = VoiceChatSim->GetSuccessfulTestCount();
+					const int32 FailedSineTest = VoiceChatSim->GetFailedTestCount();
+					ASSERT_THAT(IsTrue(SuccessfulSineTest > 0 && FailedSineTest == 0));
+				});
+	}
 };
 
 #endif // WITH_DEV_AUTOMATION_TESTS
\ No newline at end of file
diff --git a/Engine/Source/Runtime/AudioLink/AudioLinkEngine/Private/BufferedListenerBase.cpp b/Engine/Source/Runtime/AudioLink/AudioLinkEngine/Private/BufferedListenerBase.cpp
--- a/Engine/Source/Runtime/AudioLink/AudioLinkEngine/Private/BufferedListenerBase.cpp
+++ b/Engine/Source/Runtime/AudioLink/AudioLinkEngine/Private/BufferedListenerBase.cpp
@@ -167,3 +167,15 @@ int32 FBufferedListenerBase::GetNumBufferedSamples()
 	FScopeLock Lock(&CircularBuffer.GetCriticialSection());
 	return CircularBuffer.Num();
 }
+
+void FBufferedListenerBase::FlushAndResetFormat()
+{
+	CircularBuffer.Flush();
+	ResetFormat();
+}
+
+void FBufferedListenerBase::FlushAndEnsureCapacity(int32 InMinCapacity)
+{
+	CircularBuffer.EnsureCapacity(InMinCapacity);
+	ResetFormat();
+}
diff --git a/Engine/Source/Runtime/AudioLink/AudioLinkEngine/Private/DownmixedBufferedSubmixListener.cpp b/Engine/Source/Runtime/AudioLink/AudioLinkEngine/Private/DownmixedBufferedSubmixListener.cpp
--- a/Engine/Source/Runtime/AudioLink/AudioLinkEngine/Private/DownmixedBufferedSubmixListener.cpp
+++ b/Engine/Source/Runtime/AudioLink/AudioLinkEngine/Private/DownmixedBufferedSubmixListener.cpp
@@ -45,6 +45,23 @@ void FDownmixedBufferedSubmixListener::EnableSkipSilentBuffers(bool bInSkipSilen
 	bSkipSilentBuffers = bInSkipSilentBuffers;
 }
 
+void FDownmixedBufferedSubmixListener::SetOutputFormat(int32 InNumOutputChannels, int32 InSampleRate)
+{
+	check(InNumOutputChannels > 0 && InNumOutputChannels <= DownmixedBufferedSubmixListenerPrivate::MaxInputChannels);
+	check(InSampleRate > 0);
+
+	// Store the pending format FIRST so the render thread picks it up even if it fires
+	// before the flush completes. Worst case: one buffer of new-format data gets flushed
+	// (clean underrun / silence), rather than old-format data left in the ring (garbled audio).
+	PendingOutputChannels.store(InNumOutputChannels, std::memory_order_relaxed);
+	PendingOutputSampleRate.store(InSampleRate, std::memory_order_release);
+
+	// Flush the circular buffer so consumers never read stale old-format data.
+	// Grow capacity if the new format requires it (never shrinks).
+	const int32 RequiredCapacity = InNumOutputChannels * InSampleRate;
+	FlushAndEnsureCapacity(RequiredCapacity);
+}
+
 void FDownmixedBufferedSubmixListener::OnNewSubmixBuffer(const USoundSubmix* OwningSubmix, float* AudioData, int32 InNumSamples, int32 InNumChannels,
 	const int32 InSampleRate, double /*AudioClock*/)
 {
@@ -62,6 +79,25 @@ void FDownmixedBufferedSubmixListener::OnNewSubmixBuffer(const USoundSubmix* Own
 
 	if (IsStartedNonAtomic())
 	{
+		// Pick up a pending output format change from SetOutputFormat()
+		const int32 NewOutputSampleRate = PendingOutputSampleRate.exchange(0, std::memory_order_acquire);
+		if (NewOutputSampleRate > 0)
+		{
+			const int32 NewOutputChannels = PendingOutputChannels.load(std::memory_order_relaxed);
+
+			if (NewOutputChannels != NumOutputChannels)
+			{
+				Resampler.Reset(NewOutputChannels);
+			}
+
+			NumOutputChannels = NewOutputChannels;
+			OutputSampleRate = NewOutputSampleRate;
+
+			// Force rebuild of input-dependent state on this callback
+			InputSampleRate = 0;
+			NumInputChannels = 0;
+		}
+
 		// Rebuild the sample rate converter if the mixer sample rate changed
 		if (InSampleRate != InputSampleRate)
 		{
diff --git a/Engine/Source/Runtime/AudioLink/AudioLinkEngine/Public/BufferedListenerBase.h b/Engine/Source/Runtime/AudioLink/AudioLinkEngine/Public/BufferedListenerBase.h
--- a/Engine/Source/Runtime/AudioLink/AudioLinkEngine/Public/BufferedListenerBase.h
+++ b/Engine/Source/Runtime/AudioLink/AudioLinkEngine/Public/BufferedListenerBase.h
@@ -91,6 +91,36 @@ public:
 		}
 		return CanPush;
 	}
+
+	// Reset read/write pointers, discarding all buffered data without reallocating.
+	void Flush()
+	{
+		FScopeLock Lock(&CS);
+		Read = 0;
+		Write = 0;
+	}
+
+	// Grow the buffer to at least InMinCapacity if current capacity is insufficient. Never shrinks.
+	// Flushes read/write pointers (does not preserve contents).
+	void EnsureCapacity(int32 InMinCapacity)
+	{
+		check(InMinCapacity > 1);
+
+		if (!FMath::IsPowerOfTwo(InMinCapacity))
+		{
+			InMinCapacity = FMath::RoundUpToPowerOfTwo(InMinCapacity);
+		}
+
+		FScopeLock Lock(&CS);
+		if (InMinCapacity > Buffer.Num())
+		{
+			Buffer.SetNumZeroed(InMinCapacity, EAllowShrinking::No);
+			Mask = InMinCapacity - 1;
+		}
+		Read = 0;
+		Write = 0;
+	}
+
 private:
 	uint32 Read = 0;		// These grow indefinitely until wrap at 2^32,
 	uint32 Write = 0;		// this allows us to use full capacity as write >= read.
@@ -148,6 +178,12 @@ protected:
 	//* Return the number of samples able to be popped off the queue */
 	AUDIOLINKENGINE_API int32 GetNumBufferedSamples();
 
+	//* Flush the circular buffer and reset the known format so OnFormatKnown re-fires on the next buffer */
+	AUDIOLINKENGINE_API void FlushAndResetFormat();
+
+	//* Flush the circular buffer, grow capacity if needed (never shrinks), and reset format */
+	AUDIOLINKENGINE_API void FlushAndEnsureCapacity(int32 InMinCapacity);
+
 private:
 	AUDIOLINKENGINE_API void PushSilence(int32 InNumSamplesOfSilence);
 
diff --git a/Engine/Source/Runtime/AudioLink/AudioLinkEngine/Public/DownmixedBufferedSubmixListener.h b/Engine/Source/Runtime/AudioLink/AudioLinkEngine/Public/DownmixedBufferedSubmixListener.h
--- a/Engine/Source/Runtime/AudioLink/AudioLinkEngine/Public/DownmixedBufferedSubmixListener.h
+++ b/Engine/Source/Runtime/AudioLink/AudioLinkEngine/Public/DownmixedBufferedSubmixListener.h
@@ -2,6 +2,8 @@
 
 #pragma once
 
+#include <atomic>
+
 #include "AudioDevice.h"
 #include "BufferedSubmixListener.h"
 #include "DSP/RuntimeResampler.h"
@@ -50,6 +52,18 @@ public:
 	 */
 	AUDIOLINKENGINE_API void EnableSkipSilentBuffers(bool bInSkipSilentBuffers);
 
+	/** Change the output format of the listener. Safe to call from any thread.
+	 * Flushes the circular buffer immediately so that subsequent PopBuffer()
+	 * calls will not return stale data in the old format.
+	 * The audio render thread will pick up the new format on the next
+	 * OnNewSubmixBuffer() callback and rebuild internal DSP state accordingly.
+	 * Callers should expect up to one buffer period of silence after the transition.
+	 *
+	 * @param  InNumOutputChannels: New output channel count (max 8)
+	 * @param  InSampleRate: New output sample rate
+	 */
+	AUDIOLINKENGINE_API void SetOutputFormat(int32 InNumOutputChannels, int32 InSampleRate);
+
 private:
 	//~ Begin ISubmixBufferListener
 	AUDIOLINKENGINE_API void OnNewSubmixBuffer(const USoundSubmix* OwningSubmix, float* AudioData, int32 InNumSamples, int32 InNumChannels, const int32 InSampleRate, double) override;
@@ -59,6 +73,10 @@ private:
 	int32 OutputSampleRate;
 	bool bSkipSilentBuffers = false;
 
+	// Pending output format change, set by SetOutputFormat(), consumed by OnNewSubmixBuffer()
+	std::atomic<int32> PendingOutputChannels{0};
+	std::atomic<int32> PendingOutputSampleRate{0};  // 0 = no pending change (gate variable)
+
 	TArray<float> ChannelGainMap;
 	Audio::FAlignedFloatBuffer DownMixBuffer;
 	Audio::FAlignedFloatBuffer RateConvertedBuffer;
